Updated touch controls on Rolling Code template to be consistent with Rolling BP. The user can now use the virtual joystick and jump at the same time like the BP version.

diff --git a/Templates/TP_Rolling/Source/TP_Rolling/TP_RollingBall.cpp b/Templates/TP_Rolling/Source/TP_Rolling/TP_RollingBall.cpp
--- a/Templates/TP_Rolling/Source/TP_Rolling/TP_RollingBall.cpp
+++ b/Templates/TP_Rolling/Source/TP_Rolling/TP_RollingBall.cpp
@@ -48,6 +48,10 @@ void ATP_RollingBall::SetupPlayerInputComponent(class UInputComponent* PlayerInp
 	PlayerInputComponent->BindAxis("MoveForward", this, &ATP_RollingBall::MoveForward);
 
 	PlayerInputComponent->BindAction("Jump", IE_Pressed, this, &ATP_RollingBall::Jump);
+
+	// handle touch devices
+	PlayerInputComponent->BindTouch(IE_Pressed, this, &ATP_RollingBall::TouchStarted);
+	PlayerInputComponent->BindTouch(IE_Released, this, &ATP_RollingBall::TouchStopped);
 }
 
 void ATP_RollingBall::MoveRight(float Val)
@@ -78,3 +82,24 @@ void ATP_RollingBall::NotifyHit(class UPrimitiveComponent* MyComp, class AActor*
 
 	bCanJump = true;
 }
+
+void ATP_RollingBall::TouchStarted(ETouchIndex::Type FingerIndex, FVector Location)
+{
+	if (bCanJump)
+	{
+		const FVector Impulse = FVector(0.f, 0.f, JumpImpulse);
+		Ball->AddImpulse(Impulse);
+		bCanJump = false;
+	}
+
+}
+
+void ATP_RollingBall::TouchStopped(ETouchIndex::Type FingerIndex, FVector Location)
+{
+	if (bCanJump)
+	{
+		const FVector Impulse = FVector(0.f, 0.f, JumpImpulse);
+		Ball->AddImpulse(Impulse);
+		bCanJump = false;
+	}
+}
diff --git a/Templates/TP_Rolling/Source/TP_Rolling/TP_RollingBall.h b/Templates/TP_Rolling/Source/TP_Rolling/TP_RollingBall.h
--- a/Templates/TP_Rolling/Source/TP_Rolling/TP_RollingBall.h
+++ b/Templates/TP_Rolling/Source/TP_Rolling/TP_RollingBall.h
@@ -53,6 +53,12 @@ protected:
 	virtual void SetupPlayerInputComponent(class UInputComponent* InputComponent) override;
 	// End of APawn interface
 
+	/** Handler for when a touch input begins. */
+	void TouchStarted(ETouchIndex::Type FingerIndex, FVector Location);
+
+	/** Handler for when a touch input stops. */
+	void TouchStopped(ETouchIndex::Type FingerIndex, FVector Location);
+
 public:
 	/** Returns Ball subobject **/
 	FORCEINLINE class UStaticMeshComponent* GetBall() const { return Ball; }
